Live Link Hub: Enable slate/take number changes to be received from Unreal Editor devices.

Forwarding of these events to the originating device is suppressed to prevent a feedback loop.

#Synced-CL 52557823

diff --git a/Engine/Plugins/Animation/LiveLinkUnrealDevice/Source/LiveLinkUnrealDevice/LiveLinkUnrealDevice.Build.cs b/Engine/Plugins/Animation/LiveLinkUnrealDevice/Source/LiveLinkUnrealDevice/LiveLinkUnrealDevice.Build.cs
--- a/Engine/Plugins/Animation/LiveLinkUnrealDevice/Source/LiveLinkUnrealDevice/LiveLinkUnrealDevice.Build.cs
+++ b/Engine/Plugins/Animation/LiveLinkUnrealDevice/Source/LiveLinkUnrealDevice/LiveLinkUnrealDevice.Build.cs
@@ -23,6 +23,7 @@ public class LiveLinkUnrealDevice : ModuleRules
 				"Slate",
 				"SlateCore",
 				"TakeRecorder",
+				"TakesCore",
 			}
 		);
 
diff --git a/Engine/Plugins/Animation/LiveLinkUnrealDevice/Source/LiveLinkUnrealDevice/Private/Devices/LiveLinkUnrealDevice.cpp b/Engine/Plugins/Animation/LiveLinkUnrealDevice/Source/LiveLinkUnrealDevice/Private/Devices/LiveLinkUnrealDevice.cpp
--- a/Engine/Plugins/Animation/LiveLinkUnrealDevice/Source/LiveLinkUnrealDevice/Private/Devices/LiveLinkUnrealDevice.cpp
+++ b/Engine/Plugins/Animation/LiveLinkUnrealDevice/Source/LiveLinkUnrealDevice/Private/Devices/LiveLinkUnrealDevice.cpp
@@ -192,6 +192,11 @@ bool ULiveLinkUnrealDevice::IsRecording_Implementation() const
 
 void ULiveLinkUnrealDevice::HandleSlateNameChanged(FStringView InSlateName)
 {
+	if (bSuppressSlateNameSend)
+	{
+		return;
+	}
+
 	if (MessageEndpoint && DeviceAddress.IsValid())
 	{
 		FLiveLinkTakeRecorderCmd_SetSlateName* Message = FMessageEndpoint::MakeMessage<FLiveLinkTakeRecorderCmd_SetSlateName>();
@@ -203,6 +208,11 @@ void ULiveLinkUnrealDevice::HandleSlateNameChanged(FStringView InSlateName)
 
 void ULiveLinkUnrealDevice::HandleTakeNumberChanged(int32 InTakeNumber)
 {
+	if (bSuppressTakeNumberSend)
+	{
+		return;
+	}
+
 	if (MessageEndpoint && DeviceAddress.IsValid())
 	{
 		FLiveLinkTakeRecorderCmd_SetTakeNumber* Message = FMessageEndpoint::MakeMessage<FLiveLinkTakeRecorderCmd_SetTakeNumber>();
@@ -298,6 +308,50 @@ void ULiveLinkUnrealDevice::HandleMessageCatchall(const TSharedRef<IMessageConte
 					GetSettings()->DisplayName, bWasRecording, bStoppedRecording);
 			}
 		}
+		else if (MessageTypeInfo == FLiveLinkTakeRecorderCmd_SetSlateName::StaticStruct())
+		{
+			const FLiveLinkTakeRecorderCmd_SetSlateName* Cmd =
+				static_cast<const FLiveLinkTakeRecorderCmd_SetSlateName*>(InContext->GetMessage());
+
+			UE_LOGFMT(LogLiveLinkUnrealDevice, Log, "{DeviceName}: SetSlateName(\"{SlateName}\") from {RemoteAddress}",
+				GetSettings()->DisplayName, Cmd->SlateName, InContext->GetSender().ToString());
+
+			if (GetDeviceSettings<ULiveLinkUnrealDeviceSettings>()->bHasTakeMetadataAuthority)
+			{
+				ILiveLinkRecordingSession& Session = ILiveLinkRecordingSession::Get();
+				const FString PrevSlateName = Session.GetSlateName();
+
+				bSuppressSlateNameSend = true;
+				Session.SetSlateName(Cmd->SlateName);
+				bSuppressSlateNameSend = false;
+
+				UE_LOGFMT(LogLiveLinkUnrealDevice, Log, "{DeviceName}: Has take metadata authority; "
+					"PrevSlateName={PrevSlateName}, NewSlateName={NewSlateName}",
+					GetSettings()->DisplayName, PrevSlateName, Cmd->SlateName);
+			}
+		}
+		else if (MessageTypeInfo == FLiveLinkTakeRecorderCmd_SetTakeNumber::StaticStruct())
+		{
+			const FLiveLinkTakeRecorderCmd_SetTakeNumber* Cmd =
+				static_cast<const FLiveLinkTakeRecorderCmd_SetTakeNumber*>(InContext->GetMessage());
+
+			UE_LOGFMT(LogLiveLinkUnrealDevice, Log, "{DeviceName}: SetTakeNumber({TakeNumber}) from {RemoteAddress}",
+				GetSettings()->DisplayName, Cmd->TakeNumber, InContext->GetSender().ToString());
+
+			if (GetDeviceSettings<ULiveLinkUnrealDeviceSettings>()->bHasTakeMetadataAuthority)
+			{
+				ILiveLinkRecordingSession& Session = ILiveLinkRecordingSession::Get();
+				const int32 PrevTakeNumber = Session.GetTakeNumber();
+
+				bSuppressTakeNumberSend = true;
+				Session.SetTakeNumber(Cmd->TakeNumber);
+				bSuppressTakeNumberSend = false;
+
+				UE_LOGFMT(LogLiveLinkUnrealDevice, Log, "{DeviceName}: Has take metadata authority; "
+					"PrevTakeNumber={PrevTakeNumber}, NewTakeNumber={NewTakeNumber}",
+					GetSettings()->DisplayName, PrevTakeNumber, Cmd->TakeNumber);
+			}
+		}
 	}
 }
 
diff --git a/Engine/Plugins/Animation/LiveLinkUnrealDevice/Source/LiveLinkUnrealDevice/Private/Devices/LiveLinkUnrealDevice.h b/Engine/Plugins/Animation/LiveLinkUnrealDevice/Source/LiveLinkUnrealDevice/Private/Devices/LiveLinkUnrealDevice.h
--- a/Engine/Plugins/Animation/LiveLinkUnrealDevice/Source/LiveLinkUnrealDevice/Private/Devices/LiveLinkUnrealDevice.h
+++ b/Engine/Plugins/Animation/LiveLinkUnrealDevice/Source/LiveLinkUnrealDevice/Private/Devices/LiveLinkUnrealDevice.h
@@ -40,6 +40,10 @@ public:
 	/** If true, a "recording stopped" event from this device will also cause Live Link Hub to stop recording. */
 	UPROPERTY(EditAnywhere, Category="Unreal Device", AdvancedDisplay)
 	bool bHasRecordStopAuthority = true;
+
+	/** If true, slate/take changed events from this device will also cause Live Link Hub to update. */
+	UPROPERTY(EditAnywhere, Category="Unreal Device", AdvancedDisplay)
+	bool bHasTakeMetadataAuthority = true;
 };
 
 
@@ -125,4 +129,12 @@ private:
 
 	/** Recording state is updated in response to device event messages. */
 	bool bIsRecording = false;
+
+	/**
+	 * Guards against feedback loops when receiving slate/take changes from UE.
+	 * Set true while processing an incoming command so that the resulting
+	 * ILiveLinkRecordingSession delegate doesn't re-send the same command back.
+	 */
+	bool bSuppressSlateNameSend = false;
+	bool bSuppressTakeNumberSend = false;
 };
diff --git a/Engine/Plugins/Animation/LiveLinkUnrealDevice/Source/LiveLinkUnrealDevice/Private/LiveLinkHubUnrealDeviceAux.cpp b/Engine/Plugins/Animation/LiveLinkUnrealDevice/Source/LiveLinkUnrealDevice/Private/LiveLinkHubUnrealDeviceAux.cpp
--- a/Engine/Plugins/Animation/LiveLinkUnrealDevice/Source/LiveLinkUnrealDevice/Private/LiveLinkHubUnrealDeviceAux.cpp
+++ b/Engine/Plugins/Animation/LiveLinkUnrealDevice/Source/LiveLinkUnrealDevice/Private/LiveLinkHubUnrealDeviceAux.cpp
@@ -8,6 +8,7 @@
 #include "Logging/StructuredLog.h"
 #include "MessageEndpointBuilder.h"
 #include "Recorder/TakeRecorderSubsystem.h"
+#include "TakeMetaData.h"
 
 
 FLiveLinkHubUnrealDeviceAuxManager::FLiveLinkHubUnrealDeviceAuxManager()
@@ -35,6 +36,9 @@ FLiveLinkHubUnrealDeviceAuxManager::~FLiveLinkHubUnrealDeviceAuxManager()
 		HubMessagingModule->UnregisterAuxChannelRequestHandler<FLiveLinkUnrealDeviceAuxChannelRequestMessage>();
 	}
 
+	UTakeMetaData::OnTakeSlateChanged().RemoveAll(this);
+	UTakeMetaData::OnTakeNumberChanged().RemoveAll(this);
+
 	if (GEngine)
 	{
 		if (UTakeRecorderSubsystem* TakeRecorder = GEngine->GetEngineSubsystem<UTakeRecorderSubsystem>())
@@ -103,6 +107,9 @@ void FLiveLinkHubUnrealDeviceAuxManager::RegisterTakeRecorderDelegates()
 		TakeRecorder->GetOnRecordingStartedEvent().AddRaw(this, &FLiveLinkHubUnrealDeviceAuxManager::OnRecordingStarted);
 		TakeRecorder->GetOnRecordingStoppedEvent().AddRaw(this, &FLiveLinkHubUnrealDeviceAuxManager::OnRecordingStopped);
 	}
+
+	UTakeMetaData::OnTakeSlateChanged().AddRaw(this, &FLiveLinkHubUnrealDeviceAuxManager::OnSlateNameChanged);
+	UTakeMetaData::OnTakeNumberChanged().AddRaw(this, &FLiveLinkHubUnrealDeviceAuxManager::OnTakeNumberChanged);
 }
 
 
@@ -178,7 +185,10 @@ void FLiveLinkHubUnrealDeviceAuxManager::HandleSetSlateName(
 
 	if (UTakeRecorderSubsystem* TakeRecorder = GEngine->GetEngineSubsystem<UTakeRecorderSubsystem>(); ensure(TakeRecorder))
 	{
+		// Note the sender so that the resulting OnSlateNameChanged broadcast excludes it.
+		SlateNameChangeSource = InContext->GetSender();
 		TakeRecorder->SetSlateName(InCmd.SlateName);
+		SlateNameChangeSource.Reset();
 	}
 }
 
@@ -197,7 +207,10 @@ void FLiveLinkHubUnrealDeviceAuxManager::HandleSetTakeNumber(
 
 	if (UTakeRecorderSubsystem* TakeRecorder = GEngine->GetEngineSubsystem<UTakeRecorderSubsystem>(); ensure(TakeRecorder))
 	{
+		// Note the sender so that the resulting OnTakeNumberChanged broadcast excludes it.
+		TakeNumberChangeSource = InContext->GetSender();
 		TakeRecorder->SetTakeNumber(InCmd.TakeNumber);
+		TakeNumberChangeSource.Reset();
 	}
 }
 
@@ -242,19 +255,57 @@ void FLiveLinkHubUnrealDeviceAuxManager::HandleStopRecording(
 
 void FLiveLinkHubUnrealDeviceAuxManager::OnRecordingStarted(UTakeRecorder* InRecorder)
 {
-	TArray<FMessageAddress> ChannelAddresses;
-	ChannelToAddress.GenerateValueArray(ChannelAddresses);
-
 	FLiveLinkTakeRecorderEvent_RecordingStarted* EventMessage = FMessageEndpoint::MakeMessage<FLiveLinkTakeRecorderEvent_RecordingStarted>();
-	MessageEndpoint->Send(EventMessage, EMessageFlags::Reliable, {}, nullptr, ChannelAddresses, FTimespan::Zero(), FDateTime::MaxValue());
+	BroadcastToChannels(EventMessage, FLiveLinkTakeRecorderEvent_RecordingStarted::StaticStruct());
 }
 
 
 void FLiveLinkHubUnrealDeviceAuxManager::OnRecordingStopped(UTakeRecorder* InRecorder)
 {
-	TArray<FMessageAddress> ChannelAddresses;
-	ChannelToAddress.GenerateValueArray(ChannelAddresses);
-
 	FLiveLinkTakeRecorderEvent_RecordingStopped* EventMessage = FMessageEndpoint::MakeMessage<FLiveLinkTakeRecorderEvent_RecordingStopped>();
-	MessageEndpoint->Send(EventMessage, EMessageFlags::Reliable, {}, nullptr, ChannelAddresses, FTimespan::Zero(), FDateTime::MaxValue());
+	BroadcastToChannels(EventMessage, FLiveLinkTakeRecorderEvent_RecordingStopped::StaticStruct());
+}
+
+
+void FLiveLinkHubUnrealDeviceAuxManager::OnSlateNameChanged(const FString& InSlateName, UTakeMetaData* InTakeMetaData)
+{
+	UE_LOGFMT(LogLiveLinkUnrealDevice, Log, "Slate name changed to \"{SlateName}\", broadcasting to channels",
+		InSlateName);
+
+	FLiveLinkTakeRecorderCmd_SetSlateName* Message = FMessageEndpoint::MakeMessage<FLiveLinkTakeRecorderCmd_SetSlateName>();
+	Message->SlateName = InSlateName;
+	BroadcastToChannels(Message, FLiveLinkTakeRecorderCmd_SetSlateName::StaticStruct(), SlateNameChangeSource);
+}
+
+
+void FLiveLinkHubUnrealDeviceAuxManager::OnTakeNumberChanged(int32 InTakeNumber, UTakeMetaData* InTakeMetaData)
+{
+	UE_LOGFMT(LogLiveLinkUnrealDevice, Log, "Take number changed to {TakeNumber}, broadcasting to channels",
+		InTakeNumber);
+
+	FLiveLinkTakeRecorderCmd_SetTakeNumber* Message = FMessageEndpoint::MakeMessage<FLiveLinkTakeRecorderCmd_SetTakeNumber>();
+	Message->TakeNumber = InTakeNumber;
+	BroadcastToChannels(Message, FLiveLinkTakeRecorderCmd_SetTakeNumber::StaticStruct(), TakeNumberChangeSource);
+}
+
+
+void FLiveLinkHubUnrealDeviceAuxManager::BroadcastToChannels(
+	void* InMessage,
+	UScriptStruct* InTypeInfo,
+	const TOptional<FMessageAddress>& InExclude
+)
+{
+	TArray<FMessageAddress> Recipients;
+	Recipients.Reserve(ChannelToAddress.Num());
+
+	for (const TPair<FGuid, FMessageAddress>& Pair : ChannelToAddress)
+	{
+		if (!InExclude.IsSet() || Pair.Value != InExclude.GetValue())
+		{
+			Recipients.Add(Pair.Value);
+		}
+	}
+
+	MessageEndpoint->Send(InMessage, InTypeInfo, EMessageFlags::Reliable, {}, nullptr, Recipients,
+		FTimespan::Zero(), FDateTime::MaxValue());
 }
diff --git a/Engine/Plugins/Animation/LiveLinkUnrealDevice/Source/LiveLinkUnrealDevice/Private/LiveLinkHubUnrealDeviceAux.h b/Engine/Plugins/Animation/LiveLinkUnrealDevice/Source/LiveLinkUnrealDevice/Private/LiveLinkHubUnrealDeviceAux.h
--- a/Engine/Plugins/Animation/LiveLinkUnrealDevice/Source/LiveLinkUnrealDevice/Private/LiveLinkHubUnrealDeviceAux.h
+++ b/Engine/Plugins/Animation/LiveLinkUnrealDevice/Source/LiveLinkUnrealDevice/Private/LiveLinkHubUnrealDeviceAux.h
@@ -14,12 +14,13 @@ struct FLiveLinkTakeRecorderCmd_StartRecording;
 struct FLiveLinkTakeRecorderCmd_StopRecording;
 class FMessageEndpoint;
 class IMessageContext;
+class UTakeMetaData;
 class UTakeRecorder;
 
 
 /**
  * Configures a Live Link Hub auxiliary message channel for communicating with Unreal devices.
- * 
+ *
  * Handles incoming Take Recorder commands, and relays events to connected devices.
  */
 class FLiveLinkHubUnrealDeviceAuxManager
@@ -36,6 +37,8 @@ private:
 	// Take Recorder event handlers
 	void OnRecordingStarted(UTakeRecorder* InRecorder);
 	void OnRecordingStopped(UTakeRecorder* InRecorder);
+	void OnSlateNameChanged(const FString& InSlateName, UTakeMetaData* InTakeMetaData);
+	void OnTakeNumberChanged(int32 InTakeNumber, UTakeMetaData* InTakeMetaData);
 
 	// Message handlers
 	bool IsKnownSender(const FMessageAddress& InAddress) const;
@@ -47,9 +50,20 @@ private:
 	void HandleStartRecording(const FLiveLinkTakeRecorderCmd_StartRecording& InCmd, const TSharedRef<IMessageContext, ESPMode::ThreadSafe>& InContext);
 	void HandleStopRecording(const FLiveLinkTakeRecorderCmd_StopRecording& InCmd, const TSharedRef<IMessageContext, ESPMode::ThreadSafe>& InContext);
 
+	/** Broadcast a message to all connected channels, optionally excluding one address. */
+	void BroadcastToChannels(void* InMessage, UScriptStruct* InTypeInfo, const TOptional<FMessageAddress>& InExclude = {});
+
 private:
 	TSharedPtr<FMessageEndpoint> MessageEndpoint;
 
 	TMap<FGuid, FMessageAddress> ChannelToAddress;
 	TMap<FMessageAddress, FGuid> AddressToChannel;
+
+	/**
+	 * When processing an incoming slate/take change from a remote device, we note the sender
+	 * address here so that the corresponding outgoing broadcast (triggered by the local delegate)
+	 * can exclude that sender and avoid a feedback loop.
+	 */
+	TOptional<FMessageAddress> SlateNameChangeSource;
+	TOptional<FMessageAddress> TakeNumberChangeSource;
 };
